file I missed with yesterday's movie submission
#ue4
#ios

diff --git a/Engine/Plugins/Runtime/AppleMoviePlayer/Source/AppleMoviePlayer/Private/AppleMovieStreamer.cpp b/Engine/Plugins/Runtime/AppleMoviePlayer/Source/AppleMoviePlayer/Private/AppleMovieStreamer.cpp
--- a/Engine/Plugins/Runtime/AppleMoviePlayer/Source/AppleMoviePlayer/Private/AppleMovieStreamer.cpp
+++ b/Engine/Plugins/Runtime/AppleMoviePlayer/Source/AppleMoviePlayer/Private/AppleMovieStreamer.cpp
@@ -161,6 +161,10 @@ bool FAVPlayerMovieStreamer::Tick(float DeltaTime)
                 return true;
             }
         }
+        else
+        {
+            return MovieQueue.Num() == 0;
+        }
     }
 
     // Not completed.
@@ -208,10 +212,16 @@ bool FAVPlayerMovieStreamer::StartNextMovie()
 		FString MoviePath = FPaths::GameContentDir() + TEXT("Movies/") + MovieQueue[0] + TEXT(".") + FString(MOVIE_FILE_EXTENSION);
 		nsURL = [NSURL fileURLWithPath:MoviePath.GetNSString()];
 #else
-		nsURL = [[NSBundle mainBundle] URLForResource: MovieQueue[0].GetNSString() withExtension: MOVIE_FILE_EXTENSION];
+        NSString* moviestring = MovieQueue[0].GetNSString();
+		nsURL = [[NSBundle mainBundle] URLForResource: moviestring withExtension: MOVIE_FILE_EXTENSION ];
 #endif
-        check( nsURL != nil );
-
+        if (nsURL == nil)
+        {
+            UE_LOG(LogMoviePlayer, Warning, TEXT("Couldn't find movie: %s"), *MovieQueue[0]);
+            MovieQueue.RemoveAt(0);
+            return false;
+        }
+        
         NSError *error = nil;
         AudioPlayer = [[AVAudioPlayer alloc] initWithContentsOfURL:nsURL error:&error];
         if ( AudioPlayer == nil )
@@ -370,12 +380,12 @@ bool FAVPlayerMovieStreamer::CheckForNextFrameAndCopy()
             SyncStatus = Behind;
         }
 
-        if( Delta < VideoRate ) 
+        if( Delta < VideoRate )
         {
              // Video in sync with audio
             SyncStatus = Ready;
         }
-        else if(SyncStatus == Ahead) 
+        else if(SyncStatus == Ahead)
         {
             // Video ahead of audio: stay in Ahead state, exit loop
             break;
