Change error about missing audio track to a warning.

diff --git a/Engine/Plugins/Runtime/AppleMoviePlayer/Source/AppleMoviePlayer/Private/AppleMovieStreamer.cpp b/Engine/Plugins/Runtime/AppleMoviePlayer/Source/AppleMoviePlayer/Private/AppleMovieStreamer.cpp
--- a/Engine/Plugins/Runtime/AppleMoviePlayer/Source/AppleMoviePlayer/Private/AppleMovieStreamer.cpp
+++ b/Engine/Plugins/Runtime/AppleMoviePlayer/Source/AppleMoviePlayer/Private/AppleMovieStreamer.cpp
@@ -226,7 +226,7 @@ bool FAVPlayerMovieStreamer::StartNextMovie()
         AudioPlayer = [[AVAudioPlayer alloc] initWithContentsOfURL:nsURL error:&error];
         if ( AudioPlayer == nil )
         {
-            UE_LOG(LogMoviePlayer, Error, TEXT("couldn't initialize Movie player audio, bad file, or possibly just no Audio"));
+            UE_LOG(LogMoviePlayer, Warning, TEXT("couldn't initialize Movie player audio, bad file, or possibly just no Audio"));
         }
 		else
 		{
