Added a critical section to AppleMovieStreamer to make sure tick is not checking if the video tracks are loaded while AVMovie is still loading them on other thread.

diff --git a/Engine/Plugins/Runtime/AppleMoviePlayer/Source/AppleMoviePlayer/Private/AppleMovieStreamer.cpp b/Engine/Plugins/Runtime/AppleMoviePlayer/Source/AppleMoviePlayer/Private/AppleMovieStreamer.cpp
--- a/Engine/Plugins/Runtime/AppleMoviePlayer/Source/AppleMoviePlayer/Private/AppleMovieStreamer.cpp
+++ b/Engine/Plugins/Runtime/AppleMoviePlayer/Source/AppleMoviePlayer/Private/AppleMovieStreamer.cpp
@@ -102,6 +102,8 @@ bool FAVPlayerMovieStreamer::Tick(float DeltaTime)
         }
     }
 
+	FScopeLock LockVideoTracksLoading(&VideoTracksLoadingLock);
+
     if( bVideoTracksLoaded )
     {
         // Remember that we were active. Used to edge detect active/not-active transitions
@@ -246,6 +248,7 @@ bool FAVPlayerMovieStreamer::StartNextMovie()
             // !!! This block will execute asynchronously !!!
 
             // Once loaded, initialize our reader object to start pulling frames.
+			FScopeLock LockVideoTracksLoading(&VideoTracksLoadingLock);
             bVideoTracksLoaded = FinishLoadingTracks();
             if( bVideoTracksLoaded && AudioPlayer != nil )
             {
diff --git a/Engine/Plugins/Runtime/AppleMoviePlayer/Source/AppleMoviePlayer/Private/AppleMovieStreamer.h b/Engine/Plugins/Runtime/AppleMoviePlayer/Source/AppleMoviePlayer/Private/AppleMovieStreamer.h
--- a/Engine/Plugins/Runtime/AppleMoviePlayer/Source/AppleMoviePlayer/Private/AppleMovieStreamer.h
+++ b/Engine/Plugins/Runtime/AppleMoviePlayer/Source/AppleMoviePlayer/Private/AppleMovieStreamer.h
@@ -63,6 +63,8 @@ private:
     bool                        bVideoTracksLoaded;
     bool                        bWasActive;
 
+	FCriticalSection			VideoTracksLoadingLock;
+
 private:
 
     /**
