Make sure Apple movie player on Mac checks if a movie file exists so it doesn't print any errors to the log when it cannot open a track for a movie that it couldn't open.

diff --git a/Engine/Plugins/Runtime/AppleMoviePlayer/Source/AppleMoviePlayer/Private/AppleMovieStreamer.cpp b/Engine/Plugins/Runtime/AppleMoviePlayer/Source/AppleMoviePlayer/Private/AppleMovieStreamer.cpp
--- a/Engine/Plugins/Runtime/AppleMoviePlayer/Source/AppleMoviePlayer/Private/AppleMovieStreamer.cpp
+++ b/Engine/Plugins/Runtime/AppleMoviePlayer/Source/AppleMoviePlayer/Private/AppleMovieStreamer.cpp
@@ -212,7 +212,10 @@ bool FAVPlayerMovieStreamer::StartNextMovie()
         NSURL* nsURL = nil;
 #if PLATFORM_MAC
 		FString MoviePath = FPaths::GameContentDir() + TEXT("Movies/") + MovieQueue[0] + TEXT(".") + FString(MOVIE_FILE_EXTENSION);
-		nsURL = [NSURL fileURLWithPath:MoviePath.GetNSString()];
+		if (FPaths::FileExists(MoviePath))
+		{
+			nsURL = [NSURL fileURLWithPath:MoviePath.GetNSString()];
+		}
 #else
         NSString* moviestring = MovieQueue[0].GetNSString();
 		nsURL = [[NSBundle mainBundle] URLForResource: moviestring withExtension: MOVIE_FILE_EXTENSION ];
