CIS fix
#ue4

diff --git a/Engine/Plugins/Runtime/AppleMoviePlayer/Source/AppleMoviePlayer/Private/AppleMovieStreamer.cpp b/Engine/Plugins/Runtime/AppleMoviePlayer/Source/AppleMoviePlayer/Private/AppleMovieStreamer.cpp
--- a/Engine/Plugins/Runtime/AppleMoviePlayer/Source/AppleMoviePlayer/Private/AppleMovieStreamer.cpp
+++ b/Engine/Plugins/Runtime/AppleMoviePlayer/Source/AppleMoviePlayer/Private/AppleMovieStreamer.cpp
@@ -68,7 +68,7 @@ void FAVPlayerMovieStreamer::ForceCompletion()
     TeardownPlayback();
 }
 
-void FAVPlayerMovieStreamer::Init(const TArray<FString>& MoviePaths)
+bool FAVPlayerMovieStreamer::Init(const TArray<FString>& MoviePaths)
 {
 	// 
 	// Initializes the streamer for audio and video playback of the given path(s).
@@ -79,7 +79,7 @@ void FAVPlayerMovieStreamer::Init(const TArray<FString>& MoviePaths)
 	MovieQueue.Append(MoviePaths);
 
 	// Play the next movie in the queue
-	StartNextMovie();
+	return StartNextMovie();
 }
 
 bool FAVPlayerMovieStreamer::Tick(float DeltaTime)
diff --git a/Engine/Plugins/Runtime/AppleMoviePlayer/Source/AppleMoviePlayer/Private/AppleMovieStreamer.h b/Engine/Plugins/Runtime/AppleMoviePlayer/Source/AppleMoviePlayer/Private/AppleMovieStreamer.h
--- a/Engine/Plugins/Runtime/AppleMoviePlayer/Source/AppleMoviePlayer/Private/AppleMovieStreamer.h
+++ b/Engine/Plugins/Runtime/AppleMoviePlayer/Source/AppleMoviePlayer/Private/AppleMovieStreamer.h
@@ -17,7 +17,7 @@ public:
 	FAVPlayerMovieStreamer();
 	virtual ~FAVPlayerMovieStreamer();
 
-	virtual void Init(const TArray<FString>& MoviePaths) override;
+	virtual bool Init(const TArray<FString>& MoviePaths) override;
 	virtual void ForceCompletion() override;
 	virtual bool Tick(float DeltaTime) override;
 	virtual TSharedPtr<class ISlateViewport> GetViewportInterface() override;
diff --git a/Engine/Plugins/Runtime/WindowsMoviePlayer/Source/WindowsMoviePlayer/Private/WindowsMovieStreamer.cpp b/Engine/Plugins/Runtime/WindowsMoviePlayer/Source/WindowsMoviePlayer/Private/WindowsMovieStreamer.cpp
--- a/Engine/Plugins/Runtime/WindowsMoviePlayer/Source/WindowsMoviePlayer/Private/WindowsMovieStreamer.cpp
+++ b/Engine/Plugins/Runtime/WindowsMoviePlayer/Source/WindowsMoviePlayer/Private/WindowsMovieStreamer.cpp
@@ -41,16 +41,18 @@ FMediaFoundationMovieStreamer::~FMediaFoundationMovieStreamer()
 	TextureFreeList.Empty();
 }
 
-void FMediaFoundationMovieStreamer::Init(const TArray<FString>& MoviePaths)
+bool FMediaFoundationMovieStreamer::Init(const TArray<FString>& MoviePaths)
 {
 	if (MoviePaths.Num() == 0)
 	{
-		return;
+		return false;
 	}
 
 	StoredMoviePaths = MoviePaths;
 
 	OpenNextMovie();
+
+	return true;
 }
 
 void FMediaFoundationMovieStreamer::ForceCompletion()
diff --git a/Engine/Plugins/Runtime/WindowsMoviePlayer/Source/WindowsMoviePlayer/Private/WindowsMovieStreamer.h b/Engine/Plugins/Runtime/WindowsMoviePlayer/Source/WindowsMoviePlayer/Private/WindowsMovieStreamer.h
--- a/Engine/Plugins/Runtime/WindowsMoviePlayer/Source/WindowsMoviePlayer/Private/WindowsMovieStreamer.h
+++ b/Engine/Plugins/Runtime/WindowsMoviePlayer/Source/WindowsMoviePlayer/Private/WindowsMovieStreamer.h
@@ -20,7 +20,7 @@ public:
 	~FMediaFoundationMovieStreamer();
 
 	/** IMovieStreamer interface */
-	virtual void Init(const TArray<FString>& MoviePaths) override;
+	virtual bool Init(const TArray<FString>& MoviePaths) override;
 	virtual void ForceCompletion() override;
 	virtual bool Tick(float DeltaTime) override;
 	virtual TSharedPtr<class ISlateViewport> GetViewportInterface() override
