Properly handle playing multiple movies with different resolutions on in AppleMovieStreamer

diff --git a/Engine/Plugins/Runtime/AppleMoviePlayer/Source/AppleMoviePlayer/Private/AppleMovieStreamer.cpp b/Engine/Plugins/Runtime/AppleMoviePlayer/Source/AppleMoviePlayer/Private/AppleMovieStreamer.cpp
--- a/Engine/Plugins/Runtime/AppleMoviePlayer/Source/AppleMoviePlayer/Private/AppleMovieStreamer.cpp
+++ b/Engine/Plugins/Runtime/AppleMoviePlayer/Source/AppleMoviePlayer/Private/AppleMovieStreamer.cpp
@@ -425,7 +425,12 @@ bool FAVPlayerMovieStreamer::CheckForNextFrameAndCopy()
         {
             check( SrcWidth > 0 && SrcHeight > 0 );
 
-            // Now that we have video information, check on texture allocation. If we don't have a texture yet, create one.
+			TArray<uint8> TempData;
+			TempData.AddZeroed(SrcWidth * SrcHeight * 4);
+			TextureData->SetRawData(SrcWidth, SrcHeight, SrcWidth * 4, TempData);
+			check( TextureData->GetRawBytesPtr() != NULL );
+
+			// Now that we have video information, check on texture allocation. If we don't have a texture yet, create one.
             if(!Texture.IsValid() || (Texture->GetWidth() != SrcWidth || Texture->GetHeight() != SrcHeight))
             {
                 // Release any resources associated with the previous texture
@@ -438,19 +443,14 @@ bool FAVPlayerMovieStreamer::CheckForNextFrameAndCopy()
                 }
 
                 // Create and initialize a new texture
-                Texture = MakeShareable(new FSlateTexture2DRHIRef(SrcWidth, SrcHeight, PF_B8G8R8A8, TextureData, TexCreate_Dynamic | TexCreate_NoTiling));
+                Texture = MakeShareable(new FSlateTexture2DRHIRef(SrcWidth, SrcHeight, PF_B8G8R8A8, nullptr, TexCreate_Dynamic | TexCreate_NoTiling, true));
                 Texture->InitResource();
 
                 // Make sure the texture is at least updated once.
                 Texture->UpdateRHI();
                 MovieViewport->SetTexture(Texture);
-            }
-
-			TArray<uint8> TempData;
-			TempData.AddZeroed(SrcWidth * SrcHeight * 4);
-            TextureData->SetRawData(SrcWidth, SrcHeight, SrcWidth * 4, TempData);
-            check( TextureData->GetRawBytesPtr() != NULL );
-        }
+			}
+		}
 
         check( TextureData->GetBytesPerPixel() > 0 );
         
