Fixed playing multiple movies with the same size in a row on Mac

diff --git a/Engine/Plugins/Runtime/AppleMoviePlayer/Source/AppleMoviePlayer/Private/AppleMovieStreamer.cpp b/Engine/Plugins/Runtime/AppleMoviePlayer/Source/AppleMoviePlayer/Private/AppleMovieStreamer.cpp
--- a/Engine/Plugins/Runtime/AppleMoviePlayer/Source/AppleMoviePlayer/Private/AppleMovieStreamer.cpp
+++ b/Engine/Plugins/Runtime/AppleMoviePlayer/Source/AppleMoviePlayer/Private/AppleMovieStreamer.cpp
@@ -185,7 +185,9 @@ float FAVPlayerMovieStreamer::GetAspectRatio() const
 
 void FAVPlayerMovieStreamer::Cleanup()
 {
-    // Schedule textures for release.
+	MovieViewport->SetTexture(NULL);
+
+	// Schedule textures for release.
     if (Texture.IsValid())
     {
         TexturesPendingDeletion.Add(Texture);
@@ -433,7 +435,9 @@ bool FAVPlayerMovieStreamer::CheckForNextFrameAndCopy()
 			// Now that we have video information, check on texture allocation. If we don't have a texture yet, create one.
             if(!Texture.IsValid() || (Texture->GetWidth() != SrcWidth || Texture->GetHeight() != SrcHeight))
             {
-                // Release any resources associated with the previous texture
+				MovieViewport->SetTexture(NULL);
+
+				// Release any resources associated with the previous texture
                 if (Texture.IsValid())
                 {
                     // *** Aren't we already in the rendering thread? Why can't we just release the resource right here, right now?
@@ -479,9 +483,6 @@ bool FAVPlayerMovieStreamer::CheckForNextFrameAndCopy()
 
 void FAVPlayerMovieStreamer::TeardownPlayback()
 {
-    // Make sure the movie viewport doesn't hold on to the texture any longer
-    MovieViewport->SetTexture(NULL);
-
     if( LatestSamples != NULL )
     {
         CFRelease( LatestSamples );
