UE-63263 Startup movies are not playing on tvOS

#4.21
#tvOS

Startup movies need AppleMoviePlayer.uplugin, which was never enabled for tvOS.
Same behavior as an iPhone/iPad, slightly stuttering at first run, plays fine afterwards

diff --git a/Engine/Plugins/Runtime/AppleMoviePlayer/AppleMoviePlayer.uplugin b/Engine/Plugins/Runtime/AppleMoviePlayer/AppleMoviePlayer.uplugin
--- a/Engine/Plugins/Runtime/AppleMoviePlayer/AppleMoviePlayer.uplugin
+++ b/Engine/Plugins/Runtime/AppleMoviePlayer/AppleMoviePlayer.uplugin
@@ -23,6 +23,7 @@
 			"WhitelistPlatforms" :
 			[
 				"IOS",
+				"TVOS",
 				"Mac"
 			]
 		}
