UE-65406 Match3 - locking the device before the startup movie can play causes the app to crash on resume
UE-65325 Black Screen after resuming a Backgrounded app

#iOS
#4.21.1

Implementing  Supend/Resume for FAVPlayerMovieStreamer (startup movie video and audio streaming)
Note: Won't use HandleApplicationHasEnteredForeground and HandleApplicationWillEnterBackground delegates because we need to catch the interrupts in LaunchIOS before initializing the engine.

diff --git a/Engine/Plugins/Runtime/AppleMoviePlayer/Source/AppleMoviePlayer/AppleMoviePlayer.Build.cs b/Engine/Plugins/Runtime/AppleMoviePlayer/Source/AppleMoviePlayer/AppleMoviePlayer.Build.cs
--- a/Engine/Plugins/Runtime/AppleMoviePlayer/Source/AppleMoviePlayer/AppleMoviePlayer.Build.cs
+++ b/Engine/Plugins/Runtime/AppleMoviePlayer/Source/AppleMoviePlayer/AppleMoviePlayer.Build.cs
@@ -14,7 +14,8 @@ namespace UnrealBuildTool.Rules
                     "MoviePlayer",
                     "RenderCore",
                     "RHI",
-                    "Slate"
+                    "Slate",
+					"ApplicationCore"
 				}
 			);
 
diff --git a/Engine/Plugins/Runtime/AppleMoviePlayer/Source/AppleMoviePlayer/Private/AppleMovieStreamer.cpp b/Engine/Plugins/Runtime/AppleMoviePlayer/Source/AppleMoviePlayer/Private/AppleMovieStreamer.cpp
--- a/Engine/Plugins/Runtime/AppleMoviePlayer/Source/AppleMoviePlayer/Private/AppleMovieStreamer.cpp
+++ b/Engine/Plugins/Runtime/AppleMoviePlayer/Source/AppleMoviePlayer/Private/AppleMovieStreamer.cpp
@@ -8,12 +8,15 @@
 #include "Misc/CommandLine.h"
 #include "Misc/ScopeLock.h"
 #include "Misc/Paths.h"
+#include "Misc/CoreDelegates.h"
 
 DEFINE_LOG_CATEGORY_STATIC(LogMoviePlayer, Log, All);
 
 #define MOVIE_FILE_EXTENSION @"mp4"
 #define TIMESCALE 1000
 
+extern bool GIsSuspended;
+
 static FString ConvertToNativePath(const FString& Filename, bool bForWrite)
 {
 	FString Result = Filename;
@@ -76,7 +79,10 @@ SyncStatus          ( Default ),
 StartTime           ( 0.0 ),
 Cursor              ( 0.0 ),
 bVideoTracksLoaded  ( false ),
-bWasActive          ( false )
+bWasActive          ( false ),
+//interrupt
+bIsMovieInterrupted	( false ),
+ResumeTime			( kCMTimeZero )
 {
 	UE_LOG(LogMoviePlayer, Log, TEXT("FAVMoviePlayer ctor..."));
 
@@ -132,6 +138,12 @@ bool FAVPlayerMovieStreamer::Init(const TArray<FString>& MoviePaths, TEnumAsByte
 
 bool FAVPlayerMovieStreamer::Tick(float DeltaTime)
 {
+	FScopeLock LockVideoTracksLoading(&VideoTracksLoadingLock);
+
+	if (bIsMovieInterrupted)
+	{
+		return false;
+	}
     // Check the list of textures pending deletion and remove any that are no longer valid
     for (int32 TextureIndex = 0; TextureIndex < TexturesPendingDeletion.Num(); )
     {
@@ -150,8 +162,6 @@ bool FAVPlayerMovieStreamer::Tick(float DeltaTime)
         }
     }
 
-	FScopeLock LockVideoTracksLoading(&VideoTracksLoadingLock);
-
     if( bVideoTracksLoaded )
     {
         // Remember that we were active. Used to edge detect active/not-active transitions
@@ -265,8 +275,12 @@ void FAVPlayerMovieStreamer::Cleanup()
 
 bool FAVPlayerMovieStreamer::StartNextMovie()
 {
-	bool bDidStartMovie = false;
 	UE_LOG(LogMoviePlayer, Verbose, TEXT("Starting next movie....") );
+	ResumeTime = kCMTimeZero;
+
+	// Reset flag to indicate the movie may have started, but isn't playing yet.
+	bVideoTracksLoaded = false;
+
     if (MovieQueue.Num() > 0)
     {
         if( AVMovie != NULL )
@@ -276,9 +290,6 @@ bool FAVPlayerMovieStreamer::StartNextMovie()
             return false;
         }
 
-        // Reset flag to indicate the movie may have started, but isn't playing yet.
-        bVideoTracksLoaded = false;
-
         NSURL* nsURL = nil;
 		FString MoviePath = FPaths::ProjectContentDir() + TEXT("Movies/") + MovieQueue[0] + TEXT(".") + FString(MOVIE_FILE_EXTENSION);
 		if (FPaths::FileExists(MoviePath))
@@ -305,35 +316,58 @@ bool FAVPlayerMovieStreamer::StartNextMovie()
 			AudioPlayer.volume = 1;
 			[AudioPlayer prepareToPlay];
 		}
+		MovieName = MovieQueue[0];
+		MovieQueue.RemoveAt(0);
 
-        // Load the Movie with the appropriate URL.
-        AVMovie = [[AVURLAsset alloc] initWithURL:nsURL options:nil];
-        MovieQueue.RemoveAt(0);
+		bVideoTracksLoaded = LoadMovie(MovieName);
 
-        // Obtain the tracks asynchronously.
-        NSArray* nsTrackKeys = @[@"tracks"];
-        [AVMovie loadValuesAsynchronouslyForKeys:nsTrackKeys completionHandler:^()
-        {
-            // !!! This block will execute asynchronously !!!
+	}
+	return bVideoTracksLoaded;
+}
 
-            // Once loaded, initialize our reader object to start pulling frames.
-			FScopeLock LockVideoTracksLoading(&VideoTracksLoadingLock);
-            bVideoTracksLoaded = FinishLoadingTracks();
-            if( bVideoTracksLoaded && AudioPlayer != nil )
-            {
-                // Good time to start the audio playing.
-                [AudioPlayer play];
-            }
+bool FAVPlayerMovieStreamer::LoadMovie(FString InMovieName)
+{
+	// Reset flag to indicate the movie may have started, but isn't playing yet.
+	bVideoTracksLoaded = false;
+	NSURL* nsURL = nil;
+	FString MoviePath = FPaths::ProjectContentDir() + TEXT("Movies/") + InMovieName + TEXT(".") + FString(MOVIE_FILE_EXTENSION);
+	if (FPaths::FileExists(MoviePath))
+	{
+		nsURL = [NSURL fileURLWithPath : ConvertToNativePath(MoviePath, false).GetNSString()];
+	}
+	if (nsURL == nil)
+	{
+		return false;
+	}
+
+	// Load the Movie with the appropriate URL.
+    AVMovie = [[AVURLAsset alloc] initWithURL:nsURL options:nil];
+		
+    // Obtain the tracks asynchronously.
+    NSArray* nsTrackKeys = @[@"tracks"];
+    [AVMovie loadValuesAsynchronouslyForKeys:nsTrackKeys completionHandler:^()
+    {
+        // !!! This block will execute asynchronously !!!
+        // Once loaded, initialize our reader object to start pulling frames.
+		FScopeLock LockVideoTracksLoading(&VideoTracksLoadingLock);
+        bVideoTracksLoaded = FinishLoadingTracks();
+
+		// Play the next movie in the queue
+		bIsMovieInterrupted = GIsSuspended;
+
+		if (!bIsMovieInterrupted && bVideoTracksLoaded && AudioPlayer != nil)
+        {
+            // Good time to start the audio playing.
+            [AudioPlayer play];
+        }
 
-            // !!!
-        }];
+        // !!!
+    }];
 
-        // Movie has started.
-        bDidStartMovie = true;
+    // Movie has started.
 		
-		UE_LOG(LogMoviePlayer, Verbose, TEXT("Started next movie.") );
-    }
-    return bDidStartMovie;
+	UE_LOG(LogMoviePlayer, Log, TEXT("Started next movie.") );
+	return true;
 }
 
 bool FAVPlayerMovieStreamer::FinishLoadingTracks()
@@ -381,6 +415,8 @@ bool FAVPlayerMovieStreamer::FinishLoadingTracks()
                 [AVReader addOutput:AVVideoOutput];
 
                 // Begin reading!
+				[AVReader setTimeRange : CMTimeRangeMake(ResumeTime, kCMTimePositiveInfinity)];
+
                 if( ![AVReader startReading] )
                 {
                     UE_LOG(LogMoviePlayer, Error, TEXT("AVReader 'startReading' returned failure."));
@@ -392,7 +428,7 @@ bool FAVPlayerMovieStreamer::FinishLoadingTracks()
                 VideoRate = 1.0f / AVVideoTrack.nominalFrameRate;
 
                 // Save the starting time.
-                StartTime = CACurrentMediaTime();
+                StartTime = CACurrentMediaTime() - CMTimeGetSeconds(ResumeTime);
 
                 // Good to go.
                 bLoadedAndReading = true;
@@ -562,23 +598,7 @@ void FAVPlayerMovieStreamer::TeardownPlayback()
         LatestSamples = NULL;
     }
     
-    // NS Object release handled by external AutoReleasePool
-    if( AVVideoOutput != nil )
-    {
-        AVVideoOutput = nil;
-    }
-    if( AVVideoTrack != nil )
-    {
-        AVVideoTrack = nil;
-    }
-    if( AVReader != nil )
-    {
-        AVReader = nil;
-    }
-    if( AVMovie != nil )
-    {
-        AVMovie = nil;
-    }
+	ReleaseMovie();
     if ( AudioPlayer != nil )
     {
         AudioPlayer = nil;
@@ -587,9 +607,31 @@ void FAVPlayerMovieStreamer::TeardownPlayback()
     // NOTE: The any textures allocated are still allocated at this point. They will get released in Cleanup()
 }
 
+void FAVPlayerMovieStreamer::ReleaseMovie()
+{
+	if (AVMovie != nil)
+	{
+		AVMovie = nil;
+	}
+	// NS Object release handled by external AutoReleasePool
+	if (AVVideoOutput != nil)
+	{
+		AVVideoOutput = nil;
+	}
+	if (AVVideoTrack != nil)
+	{
+		AVVideoTrack = nil;
+	}
+	if (AVReader != nil)
+	{
+		AVReader = nil;
+	}
+
+}
+
 FString FAVPlayerMovieStreamer::GetMovieName()
 {
-	return MovieQueue.Num() > 0 ? MovieQueue[0] : TEXT("");
+	return MovieName;
 }
 
 bool FAVPlayerMovieStreamer::IsLastMovieInPlaylist()
@@ -597,3 +639,40 @@ bool FAVPlayerMovieStreamer::IsLastMovieInPlaylist()
 	return MovieQueue.Num() <= 1;
 }
 
+
+/* Interrupts */
+void FAVPlayerMovieStreamer::Suspend()
+{
+	if (bIsMovieInterrupted)
+	{
+		//already paused
+		return;
+	}
+	bIsMovieInterrupted = true;
+
+	FScopeLock LockVideoTracksLoading(&VideoTracksLoadingLock);
+
+	if (bVideoTracksLoaded && AVReader != nil)
+	{
+		ResumeTime = CMTimeMake(Cursor * TIMESCALE, TIMESCALE);
+
+		[AudioPlayer pause];
+
+		[AVReader cancelReading];
+
+	}
+	ReleaseMovie();
+}
+
+void FAVPlayerMovieStreamer::Resume()
+{
+	if (!bIsMovieInterrupted)
+	{
+		//already resumed
+		return;
+	}
+	FScopeLock LockVideoTracksLoading(&VideoTracksLoadingLock);
+
+	LoadMovie(MovieName);
+}
+
diff --git a/Engine/Plugins/Runtime/AppleMoviePlayer/Source/AppleMoviePlayer/Private/AppleMovieStreamer.h b/Engine/Plugins/Runtime/AppleMoviePlayer/Source/AppleMoviePlayer/Private/AppleMovieStreamer.h
--- a/Engine/Plugins/Runtime/AppleMoviePlayer/Source/AppleMoviePlayer/Private/AppleMovieStreamer.h
+++ b/Engine/Plugins/Runtime/AppleMoviePlayer/Source/AppleMoviePlayer/Private/AppleMovieStreamer.h
@@ -27,6 +27,8 @@ public:
 
 	virtual FString GetMovieName() override;
 	virtual bool IsLastMovieInPlaylist() override;
+	virtual void Suspend() override;
+	virtual void Resume() override;
 
 	FOnCurrentMovieClipFinished OnCurrentMovieClipFinishedDelegate;
 	virtual FOnCurrentMovieClipFinished& OnCurrentMovieClipFinished() override { return OnCurrentMovieClipFinishedDelegate; }
@@ -61,6 +63,8 @@ private:
     AVAssetTrack*               AVVideoTrack;
     CMSampleBufferRef           LatestSamples;
 
+	FString MovieName;
+
     // AV Synchronization
     float                       VideoRate;
     int                         SyncStatus;
@@ -70,6 +74,9 @@ private:
     bool                        bVideoTracksLoaded;
     bool                        bWasActive;
 
+	bool						bIsMovieInterrupted;
+	CMTime						ResumeTime;
+
 	FCriticalSection			VideoTracksLoadingLock;
 
 private:
@@ -85,4 +92,7 @@ private:
 
     bool CheckForNextFrameAndCopy();
 
+	void ReleaseMovie();
+
+	bool LoadMovie(FString MovieName);
 };
diff --git a/Engine/Source/Runtime/Launch/Private/IOS/LaunchIOS.cpp b/Engine/Source/Runtime/Launch/Private/IOS/LaunchIOS.cpp
--- a/Engine/Source/Runtime/Launch/Private/IOS/LaunchIOS.cpp
+++ b/Engine/Source/Runtime/Launch/Private/IOS/LaunchIOS.cpp
@@ -20,6 +20,7 @@
 #include "GenericPlatform/GenericApplication.h"
 #include "Misc/ConfigCacheIni.h"
 #include "MoviePlayer.h"
+#include "PreLoadScreenManager.h"
 
 FEngineLoop GEngineLoop;
 FGameLaunchDaemonMessageHandler GCommandSystem;
@@ -36,6 +37,13 @@ FAutoConsoleVariableRef CVarDisableAudioSuspendOnAudioInterrupt(
 
 void FAppEntry::Suspend(bool bIsInterrupt)
 {
+	// also treats interrupts BEFORE initializing the engine
+	// the movie player gets initialized on the preinit phase, ApplicationHasEnteredForegroundDelegate and ApplicationWillEnterBackgroundDelegate are not yet available
+	if (GetMoviePlayer())
+	{
+		GetMoviePlayer()->Suspend();
+	}
+
 	if (GEngine && GEngine->GetMainAudioDevice())
 	{
         FAudioDevice* AudioDevice = GEngine->GetMainAudioDevice();
@@ -101,6 +109,11 @@ void FAppEntry::Suspend(bool bIsInterrupt)
 
 void FAppEntry::Resume(bool bIsInterrupt)
 {
+	if (GetMoviePlayer())
+	{
+		GetMoviePlayer()->Resume();
+	}
+
 	if (GEngine && GEngine->GetMainAudioDevice())
 	{
         FAudioDevice* AudioDevice = GEngine->GetMainAudioDevice();
diff --git a/Engine/Source/Runtime/MoviePlayer/Private/DefaultGameMoviePlayer.cpp b/Engine/Source/Runtime/MoviePlayer/Private/DefaultGameMoviePlayer.cpp
--- a/Engine/Source/Runtime/MoviePlayer/Private/DefaultGameMoviePlayer.cpp
+++ b/Engine/Source/Runtime/MoviePlayer/Private/DefaultGameMoviePlayer.cpp
@@ -859,4 +859,22 @@ void FDefaultGameMoviePlayer::ForceCompletion()
 	{
 		MovieStreamer->ForceCompletion();
 	}
-}
\ No newline at end of file
+}
+
+/*Interrupts*/
+void FDefaultGameMoviePlayer::Suspend()
+{
+	if (MovieStreamer.IsValid())
+	{
+		MovieStreamer->Suspend();
+	}
+}
+
+void FDefaultGameMoviePlayer::Resume()
+{
+	if (MovieStreamer.IsValid())
+	{
+		MovieStreamer->Resume();
+	}
+}
+
diff --git a/Engine/Source/Runtime/MoviePlayer/Private/DefaultGameMoviePlayer.h b/Engine/Source/Runtime/MoviePlayer/Private/DefaultGameMoviePlayer.h
--- a/Engine/Source/Runtime/MoviePlayer/Private/DefaultGameMoviePlayer.h
+++ b/Engine/Source/Runtime/MoviePlayer/Private/DefaultGameMoviePlayer.h
@@ -105,6 +105,8 @@ public:
 	virtual bool IsLastMovieInPlaylist() override;
 
 	virtual void ForceCompletion() override;
+	virtual void Suspend() override;
+	virtual void Resume() override;
 
 	float GetViewportDPIScale() const;
 
