[Build] UE4 Dev-Mobile - CL 4624927 Nightly Build
Compile UE4Editor Mac
Undefined symbols for architecture x86_64:
"_GIsSuspended", referenced from:
____ZN22FAVPlayerMovieStreamer9LoadMovieE7FString_block_invoke in Module.AppleMoviePlayer.cpp.o
ld: symbol(s) not found for architecture x86_64

diff --git a/Engine/Plugins/Runtime/AppleMoviePlayer/Source/AppleMoviePlayer/Private/AppleMovieStreamer.cpp b/Engine/Plugins/Runtime/AppleMoviePlayer/Source/AppleMoviePlayer/Private/AppleMovieStreamer.cpp
--- a/Engine/Plugins/Runtime/AppleMoviePlayer/Source/AppleMoviePlayer/Private/AppleMovieStreamer.cpp
+++ b/Engine/Plugins/Runtime/AppleMoviePlayer/Source/AppleMoviePlayer/Private/AppleMovieStreamer.cpp
@@ -15,7 +15,9 @@ DEFINE_LOG_CATEGORY_STATIC(LogMoviePlayer, Log, All);
 #define MOVIE_FILE_EXTENSION @"mp4"
 #define TIMESCALE 1000
 
+#if PLATFORM_IOS
 extern bool GIsSuspended;
+#endif
 
 static FString ConvertToNativePath(const FString& Filename, bool bForWrite)
 {
@@ -353,8 +355,11 @@ bool FAVPlayerMovieStreamer::LoadMovie(FString InMovieName)
         bVideoTracksLoaded = FinishLoadingTracks();
 
 		// Play the next movie in the queue
+        
+#if PLATFORM_IOS
 		bIsMovieInterrupted = GIsSuspended;
-
+#endif
+        
 		if (!bIsMovieInterrupted && bVideoTracksLoaded && AudioPlayer != nil)
         {
             // Good time to start the audio playing.
