Make sure to wait for AVMovie load completion handler before returning from FAVPlayerMovieStreamer::LoadMovie()

diff --git a/Engine/Plugins/Runtime/AppleMoviePlayer/Source/AppleMoviePlayer/Private/AppleMovieStreamer.cpp b/Engine/Plugins/Runtime/AppleMoviePlayer/Source/AppleMoviePlayer/Private/AppleMovieStreamer.cpp
--- a/Engine/Plugins/Runtime/AppleMoviePlayer/Source/AppleMoviePlayer/Private/AppleMovieStreamer.cpp
+++ b/Engine/Plugins/Runtime/AppleMoviePlayer/Source/AppleMoviePlayer/Private/AppleMovieStreamer.cpp
@@ -348,7 +348,9 @@ bool FAVPlayerMovieStreamer::LoadMovie(FString InMovieName)
 
 	// Load the Movie with the appropriate URL.
     AVMovie = [[AVURLAsset alloc] initWithURL:nsURL options:nil];
-		
+
+	__block bool bLoadCompleted = false;
+
     // Obtain the tracks asynchronously.
     NSArray* nsTrackKeys = @[@"tracks"];
     [AVMovie loadValuesAsynchronouslyForKeys:nsTrackKeys completionHandler:^()
@@ -371,12 +373,18 @@ bool FAVPlayerMovieStreamer::LoadMovie(FString InMovieName)
         }
 
         // !!!
+		bLoadCompleted = true;
     }];
 
+	while (!bLoadCompleted)
+	{
+		FPlatformProcess::Sleep(0);
+	}
+
     // Movie has started.
 		
 	UE_LOG(LogMoviePlayer, Log, TEXT("Started next movie.") );
-	return true;
+	return bVideoTracksLoaded;
 }
 
 bool FAVPlayerMovieStreamer::FinishLoadingTracks()
