UE-70375 Match3 - Backgrounding app during startup movie will freeze the application when resumed on iOS

#4.22
#iOS

On Release-4.22 LoadMovie is no longer async, since the main thread now waits for loadValuesAsynchronouslyForKeys to finish.

LoadMovie is called on init (when creating the player), on tick(when starting the next movie) and on interrupt resume. For the last 2 cases, the LockVideoTracksLoading barrier inside loadValuesAsynchronouslyForKeys will lock the main thread in an infinite loop in LoadMovie;

I've removed LockVideoTracksLoading from LoadMovie, now it's used only for syncing tick() with the interrupts.

diff --git a/Engine/Plugins/Runtime/AppleMoviePlayer/Source/AppleMoviePlayer/Private/AppleMovieStreamer.cpp b/Engine/Plugins/Runtime/AppleMoviePlayer/Source/AppleMoviePlayer/Private/AppleMovieStreamer.cpp
--- a/Engine/Plugins/Runtime/AppleMoviePlayer/Source/AppleMoviePlayer/Private/AppleMovieStreamer.cpp
+++ b/Engine/Plugins/Runtime/AppleMoviePlayer/Source/AppleMoviePlayer/Private/AppleMovieStreamer.cpp
@@ -357,7 +357,6 @@ bool FAVPlayerMovieStreamer::LoadMovie(FString InMovieName)
     {
         // !!! This block will execute asynchronously !!!
         // Once loaded, initialize our reader object to start pulling frames.
-		FScopeLock LockVideoTracksLoading(&VideoTracksLoadingLock);
         bVideoTracksLoaded = FinishLoadingTracks();
 
 		// Play the next movie in the queue
