FMac Startup movie stuttering:

1) Change FSlateLoadingSynchronizationMechanism busy wait to yeild on destory and reduce sleep on thread exit wait.
2) Optimization: Remove busy wait blocking call on load next movie - make it async again.
3) Optimization: Don't do double texture data copy - which it seems to always do if it has another frame ready - copy straight to target texture (reduces CPU time by about 30% total in its ::tick() function.

diff --git a/Engine/Plugins/Runtime/AppleMoviePlayer/Source/AppleMoviePlayer/Private/AppleMovieStreamer.cpp b/Engine/Plugins/Runtime/AppleMoviePlayer/Source/AppleMoviePlayer/Private/AppleMovieStreamer.cpp
--- a/Engine/Plugins/Runtime/AppleMoviePlayer/Source/AppleMoviePlayer/Private/AppleMovieStreamer.cpp
+++ b/Engine/Plugins/Runtime/AppleMoviePlayer/Source/AppleMoviePlayer/Private/AppleMovieStreamer.cpp
@@ -92,7 +92,6 @@ ResumeTime			( kCMTimeZero )
 {
 	UE_LOG(LogMoviePlayer, Log, TEXT("FAVMoviePlayer ctor..."));
 
-    TextureData = MakeShareable(new FSlateTextureData());
     MovieViewport = MakeShareable(new FMovieViewport());
 }
 
@@ -150,6 +149,7 @@ bool FAVPlayerMovieStreamer::Tick(float DeltaTime)
 	{
 		return false;
 	}
+
     // Check the list of textures pending deletion and remove any that are no longer valid
     for (int32 TextureIndex = 0; TextureIndex < TexturesPendingDeletion.Num(); )
     {
@@ -173,14 +173,7 @@ bool FAVPlayerMovieStreamer::Tick(float DeltaTime)
         // Remember that we were active. Used to edge detect active/not-active transitions
         bWasActive = true;
 
-        if( CheckForNextFrameAndCopy() )
-        {
-            // Copy new frame data.
-            uint32 Stride;
-            uint8* DestTextureData = (uint8*)RHILockTexture2D( Texture->GetTypedResource(), 0, RLM_WriteOnly, Stride, false );
-            FMemory::Memcpy( DestTextureData, TextureData->GetRawBytesPtr(), TextureData->GetRawBytes().Num() );
-            RHIUnlockTexture2D( Texture->GetTypedResource(), 0, false );
-        }
+		CheckForNextFrameAndCopy();
 
         check(AVReader != nil);
         AVAssetReaderStatus Status = [AVReader status];
@@ -325,14 +318,15 @@ bool FAVPlayerMovieStreamer::StartNextMovie()
 		MovieName = MovieQueue[0];
 		MovieQueue.RemoveAt(0);
 
-		bVideoTracksLoaded = LoadMovie(MovieName);
-
+		return LoadMovieAsync(MovieName);
 	}
-	return bVideoTracksLoaded;
+	return false;
 }
 
-bool FAVPlayerMovieStreamer::LoadMovie(FString InMovieName)
+bool FAVPlayerMovieStreamer::LoadMovieAsync(FString InMovieName)
 {
+	FScopeLock LockVideoTracksLoading(&VideoTracksLoadingLock);
+	
 	// Reset flag to indicate the movie may have started, but isn't playing yet.
 	bVideoTracksLoaded = false;
 	NSURL* nsURL = nil;
@@ -341,6 +335,7 @@ bool FAVPlayerMovieStreamer::LoadMovie(FString InMovieName)
 	{
 		nsURL = [NSURL fileURLWithPath : ConvertToNativePath(MoviePath, false).GetNSString()];
 	}
+	
 	if (nsURL == nil)
 	{
 		return false;
@@ -349,41 +344,21 @@ bool FAVPlayerMovieStreamer::LoadMovie(FString InMovieName)
 	// Load the Movie with the appropriate URL.
     AVMovie = [[AVURLAsset alloc] initWithURL:nsURL options:nil];
 
-	__block bool bLoadCompleted = false;
-
     // Obtain the tracks asynchronously.
     NSArray* nsTrackKeys = @[@"tracks"];
     [AVMovie loadValuesAsynchronouslyForKeys:nsTrackKeys completionHandler:^()
     {
-        // !!! This block will execute asynchronously !!!
-        // Once loaded, initialize our reader object to start pulling frames.
-        bVideoTracksLoaded = FinishLoadingTracks();
+    	FScopeLock AynscLockVideoTracksLoading(&VideoTracksLoadingLock);
 
-		// Play the next movie in the queue
+        // Once loaded, initialize our reader object to start pulling frames
+        bVideoTracksLoaded = FinishLoadingTracks();
         
 #if PLATFORM_IOS
 		bIsMovieInterrupted = GIsSuspended;
 #endif
-        
-		if (!bIsMovieInterrupted && bVideoTracksLoaded && AudioPlayer != nil)
-        {
-            // Good time to start the audio playing.
-            [AudioPlayer play];
-        }
-
-        // !!!
-		bLoadCompleted = true;
     }];
-
-	while (!bLoadCompleted)
-	{
-		FPlatformProcess::Sleep(0);
-	}
-
-    // Movie has started.
-		
-	UE_LOG(LogMoviePlayer, Log, TEXT("Started next movie.") );
-	return bVideoTracksLoaded;
+	
+    return true;
 }
 
 bool FAVPlayerMovieStreamer::FinishLoadingTracks()
@@ -443,8 +418,8 @@ bool FAVPlayerMovieStreamer::FinishLoadingTracks()
                 check( AVVideoTrack.nominalFrameRate );
                 VideoRate = 1.0f / AVVideoTrack.nominalFrameRate;
 
-                // Save the starting time.
-                StartTime = CACurrentMediaTime() - CMTimeGetSeconds(ResumeTime);
+                // Reset the starting time.
+                StartTime = 0.0;
 
                 // Good to go.
                 bLoadedAndReading = true;
@@ -475,14 +450,23 @@ bool FAVPlayerMovieStreamer::CheckForNextFrameAndCopy()
     bool bHasNewFrame = false;
 
     // We need to synchronize the video playback with the audio.
-    // If the video frame is within tolerance (Ready), update the Texture Data.
+    // If the video frame is within tolerance (Ready), update the Texture.
     // If the video is Behind, throw it away and get the next one until we catch up with the ref time.
-    // If the video is Ahead, update the TextureData but don't retrieve more frames until time catches up.
+    // If the video is Ahead, update the Texture but don't retrieve more frames until time catches up.
 
-    
+    if(StartTime == 0.0)
+    {
+    	// Now kick everything going at the same time
+		StartTime = CACurrentMediaTime() - CMTimeGetSeconds(ResumeTime);
+		
+		if(AudioPlayer != nil && !AudioPlayer.isPlaying)
+		{
+			[AudioPlayer play];
+		}
+    }
+	
     while( SyncStatus != Ready )
     {
-        double Delta;
         if( SyncStatus != Ahead )
         {
             LatestSamples = [AVVideoOutput copyNextSampleBuffer];
@@ -498,17 +482,16 @@ bool FAVPlayerMovieStreamer::CheckForNextFrameAndCopy()
 
         // Get the time since playback began
         Cursor = CACurrentMediaTime() - StartTime;
-        CMTime caCurrentTime = CMTimeMake(Cursor * TIMESCALE, TIMESCALE);
 
         // Compute delta of video frame and current playback times
-        Delta = CMTimeGetSeconds(caCurrentTime) - CMTimeGetSeconds(FrameTimeStamp);
+        double Delta = Cursor - CMTimeGetSeconds(FrameTimeStamp);
 
-        if( Delta < 0 ) 
+        if( Delta < 0.0 )
         {
-            Delta *= -1;  
+            Delta *= - 1.0;
             SyncStatus = Ahead;
         }
-        else 
+        else
         {
             SyncStatus = Behind;
         }
@@ -523,7 +506,7 @@ bool FAVPlayerMovieStreamer::CheckForNextFrameAndCopy()
             // Video ahead of audio: stay in Ahead state, exit loop
             break;
         }
-        else 
+        else
         {
             // Video behind audio (Behind): stay in loop
             CFRelease(LatestSamples);
@@ -547,17 +530,6 @@ bool FAVPlayerMovieStreamer::CheckForNextFrameAndCopy()
         uint32 SrcWidth  = (uint32)Size.width;
         uint32 SrcHeight = (uint32)Size.height;
 
-        // Now that we have video information, ensure that we have texture data in the right dimensions
-        if (TextureData->GetWidth() != SrcWidth || TextureData->GetHeight() != SrcHeight)
-        {
-            check( SrcWidth > 0 && SrcHeight > 0 );
-
-			TArray<uint8> TempData;
-			TempData.AddZeroed(SrcWidth * SrcHeight * 4);
-			TextureData->SetRawData(SrcWidth, SrcHeight, SrcWidth * 4, TempData);
-			check( TextureData->GetRawBytesPtr() != NULL );
-        }
-
         // Now that we have video information, check on texture allocation. If we don't have a texture yet, create one.
         if(!Texture.IsValid() || (Texture->GetWidth() != SrcWidth || Texture->GetHeight() != SrcHeight))
         {
@@ -580,13 +552,13 @@ bool FAVPlayerMovieStreamer::CheckForNextFrameAndCopy()
             Texture->UpdateRHI();
             MovieViewport->SetTexture(Texture);
         }
-
-        check( TextureData->GetBytesPerPixel() > 0 );
-        
-        // Copy the video data
-        uint32 Len = TextureData->GetBytesPerPixel() * SrcHeight;
-        check( Len > 0 );
-        FMemory::Memcpy(TextureData->GetRawBytesPtr(), pVideoData, Len);
+		
+		uint32 DataLen = SrcWidth * 4 * SrcHeight;
+        uint32 Stride;
+		
+		uint8* DestTextureData = (uint8*)RHILockTexture2D( Texture->GetTypedResource(), 0, RLM_WriteOnly, Stride, false );
+		FMemory::Memcpy( DestTextureData, pVideoData, DataLen );
+		RHIUnlockTexture2D( Texture->GetTypedResource(), 0, false );
 
         // Re-lock and release the video data.
         CVPixelBufferUnlockBaseAddress( pPixelBuffer, kCVPixelBufferLock_ReadOnly );
@@ -687,8 +659,6 @@ void FAVPlayerMovieStreamer::Resume()
 		//already resumed
 		return;
 	}
-	FScopeLock LockVideoTracksLoading(&VideoTracksLoadingLock);
-
-	LoadMovie(MovieName);
+	
+	LoadMovieAsync(MovieName);
 }
-
diff --git a/Engine/Plugins/Runtime/AppleMoviePlayer/Source/AppleMoviePlayer/Private/AppleMovieStreamer.h b/Engine/Plugins/Runtime/AppleMoviePlayer/Source/AppleMoviePlayer/Private/AppleMovieStreamer.h
--- a/Engine/Plugins/Runtime/AppleMoviePlayer/Source/AppleMoviePlayer/Private/AppleMovieStreamer.h
+++ b/Engine/Plugins/Runtime/AppleMoviePlayer/Source/AppleMoviePlayer/Private/AppleMovieStreamer.h
@@ -49,7 +49,6 @@ private:
 	/** Texture and viewport data for displaying to Slate */
 	TSharedPtr<FMovieViewport> MovieViewport;
 
-    TSharedPtr<FSlateTextureData, ESPMode::ThreadSafe> TextureData;
     TSharedPtr<FSlateTexture2DRHIRef, ESPMode::ThreadSafe> Texture;
 
     // The list of pending movies
@@ -94,5 +93,5 @@ private:
 
 	void ReleaseMovie();
 
-	bool LoadMovie(FString MovieName);
+	bool LoadMovieAsync(FString MovieName);
 };
diff --git a/Engine/Source/Runtime/MoviePlayer/Private/MoviePlayerThreading.cpp b/Engine/Source/Runtime/MoviePlayer/Private/MoviePlayerThreading.cpp
--- a/Engine/Source/Runtime/MoviePlayer/Private/MoviePlayerThreading.cpp
+++ b/Engine/Source/Runtime/MoviePlayer/Private/MoviePlayerThreading.cpp
@@ -75,7 +75,7 @@ void FSlateLoadingSynchronizationMechanism::DestroySlateThread()
 		{
 			FPlatformApplicationMisc::PumpMessages(false);
 
-			FPlatformProcess::Sleep(0.016f);
+			FPlatformProcess::Sleep(0.f);
 		}
 
 		delete SlateLoadingThread;
@@ -151,7 +151,7 @@ void FSlateLoadingSynchronizationMechanism::SlateThreadRunMainLoop()
 	
 	while (IsSlateDrawPassEnqueued())
 	{
-		FPlatformProcess::Sleep(0.1f);
+		FPlatformProcess::Sleep(1.f / 60.f);
 	}
 	
 	bMainLoopRunning = false;
