Fix for crash in macOS dedicated server.  Movie streamer shouldn't try to create a Metal device or a texture cache unless the movie streamer is initilaized and also make sure even if Init() is called on the server it won't do anything.

diff --git a/Engine/Plugins/Runtime/AppleMoviePlayer/Source/AppleMoviePlayer/Private/AppleMovieStreamer.cpp b/Engine/Plugins/Runtime/AppleMoviePlayer/Source/AppleMoviePlayer/Private/AppleMovieStreamer.cpp
--- a/Engine/Plugins/Runtime/AppleMoviePlayer/Source/AppleMoviePlayer/Private/AppleMovieStreamer.cpp
+++ b/Engine/Plugins/Runtime/AppleMoviePlayer/Source/AppleMoviePlayer/Private/AppleMovieStreamer.cpp
@@ -142,13 +142,6 @@ MetalTextureCache	( NULL )
 	UE_LOG(LogMoviePlayer, Log, TEXT("FAVMoviePlayer ctor..."));
 
     MovieViewport = MakeShareable(new FMovieViewport());
-	
-
-	id<MTLDevice> Device = (id<MTLDevice>)GDynamicRHI->RHIGetNativeDevice();
-	check(Device);
-
-	CVReturn Return = CVMetalTextureCacheCreate(kCFAllocatorDefault, nullptr, Device, nullptr, &MetalTextureCache);
-	check(Return == kCVReturnSuccess);
 }
 
 FAVPlayerMovieStreamer::~FAVPlayerMovieStreamer()
@@ -181,6 +174,21 @@ void FAVPlayerMovieStreamer::ForceCompletion()
 
 bool FAVPlayerMovieStreamer::Init(const TArray<FString>& MoviePaths, TEnumAsByte<EMoviePlaybackType> inPlaybackType)
 {
+	if(MetalTextureCache == NULL && FApp::CanEverRender())
+	{
+		id<MTLDevice> Device = (id<MTLDevice>)GDynamicRHI->RHIGetNativeDevice();
+		check(Device != nil);
+
+		CVReturn Return = CVMetalTextureCacheCreate(kCFAllocatorDefault, nullptr, Device, nullptr, &MetalTextureCache);
+		check(Return == kCVReturnSuccess);
+		check(MetalTextureCache != NULL);
+	}
+	
+	if(MetalTextureCache == NULL)
+	{
+		return false;
+	}
+
 	// 
 	// Initializes the streamer for audio and video playback of the given path(s).
 	// NOTE: If multiple paths are provided, it is expect that they be played back seamlessly.
