Part two fix for startup movie playback stuttering on Mac.  Use io surface backed pixel buffers that we can just use directly on Metal.  Tested on 2012 MBP (8-12FPS to 120FPS), iMac Pro (No change 120FPS), iOS (No Change 30FPS capped) (FPS as reported by Xcode tools).

diff --git a/Engine/Plugins/Runtime/AppleMoviePlayer/Source/AppleMoviePlayer/Private/AppleMovieStreamer.cpp b/Engine/Plugins/Runtime/AppleMoviePlayer/Source/AppleMoviePlayer/Private/AppleMovieStreamer.cpp
--- a/Engine/Plugins/Runtime/AppleMoviePlayer/Source/AppleMoviePlayer/Private/AppleMovieStreamer.cpp
+++ b/Engine/Plugins/Runtime/AppleMoviePlayer/Source/AppleMoviePlayer/Private/AppleMovieStreamer.cpp
@@ -69,6 +69,54 @@ static FString ConvertToNativePath(const FString& Filename, bool bForWrite)
 	return Result;
 }
 
+//
+// Clone from FAvfMediaVideoSampler
+//
+class FAppleMovieStreamerTexture2DResourceWrapper
+	: public FResourceBulkDataInterface
+{
+public:
+
+	FAppleMovieStreamerTexture2DResourceWrapper(CFTypeRef InImageBuffer)
+		: ImageBuffer(InImageBuffer)
+	{
+		check(ImageBuffer);
+		CFRetain(ImageBuffer);
+	}
+
+	virtual ~FAppleMovieStreamerTexture2DResourceWrapper()
+	{
+		CFRelease(ImageBuffer);
+		ImageBuffer = nullptr;
+	}
+
+public:
+
+	//~ FResourceBulkDataInterface interface
+
+	virtual void Discard() override
+	{
+		delete this;
+	}
+
+	virtual const void* GetResourceBulkData() const override
+	{
+		return ImageBuffer;
+	}
+	
+	virtual uint32 GetResourceBulkDataSize() const override
+	{
+		return ImageBuffer ? ~0u : 0;
+	}
+
+	virtual EBulkDataType GetResourceType() const override
+	{
+		return EBulkDataType::MediaTexture;
+	}
+	
+	CFTypeRef ImageBuffer;
+};
+
 //
 // Methods
 //
@@ -86,13 +134,20 @@ StartTime           ( 0.0 ),
 Cursor              ( 0.0 ),
 bVideoTracksLoaded  ( false ),
 bWasActive          ( false ),
-//interrupt
 bIsMovieInterrupted	( false ),
-ResumeTime			( kCMTimeZero )
+ResumeTime			( kCMTimeZero ),
+MetalTextureCache	( NULL )
 {
 	UE_LOG(LogMoviePlayer, Log, TEXT("FAVMoviePlayer ctor..."));
 
     MovieViewport = MakeShareable(new FMovieViewport());
+	
+
+	id<MTLDevice> Device = (id<MTLDevice>)GDynamicRHI->RHIGetNativeDevice();
+	check(Device);
+
+	CVReturn Return = CVMetalTextureCacheCreate(kCFAllocatorDefault, nullptr, Device, nullptr, &MetalTextureCache);
+	check(Return == kCVReturnSuccess);
 }
 
 FAVPlayerMovieStreamer::~FAVPlayerMovieStreamer()
@@ -103,10 +158,6 @@ FAVPlayerMovieStreamer::~FAVPlayerMovieStreamer()
 
     // Clean up any remaining resources
     Cleanup();
-
-    // Clear out the pending list
-    // NOTE: No guarantees here that the resources are actually freed, and we can't force them to be freed.
-    TexturesPendingDeletion.Empty();
 }
 
 void FAVPlayerMovieStreamer::ForceCompletion()
@@ -150,24 +201,6 @@ bool FAVPlayerMovieStreamer::Tick(float DeltaTime)
 		return false;
 	}
 
-    // Check the list of textures pending deletion and remove any that are no longer valid
-    for (int32 TextureIndex = 0; TextureIndex < TexturesPendingDeletion.Num(); )
-    {
-        TSharedPtr<FSlateTexture2DRHIRef, ESPMode::ThreadSafe> Tex = TexturesPendingDeletion[TextureIndex];
-        if (!Tex->IsInitialized())
-        {
-            // Texture can now be removed from the list
-            TexturesPendingDeletion.RemoveAt(TextureIndex);
-
-            // Don't increment i, since element 'i' wasn't just filled by the next element
-        }
-        else
-        {
-            // Advance to the next one
-            ++TextureIndex;
-        }
-    }
-
     if( bVideoTracksLoaded )
     {
         // Remember that we were active. Used to edge detect active/not-active transitions
@@ -261,15 +294,21 @@ void FAVPlayerMovieStreamer::Cleanup()
 		LatestSamples = NULL;
 	}
 	
-    MovieViewport->SetTexture(NULL);
-
-	// Schedule textures for release.
-    if (Texture.IsValid())
-    {
-        TexturesPendingDeletion.Add(Texture);
-        BeginReleaseResource(Texture.Get());
-        Texture.Reset();
-    }
+	MovieViewport->SetTexture(NULL);
+	
+	if(SlateTexture.IsValid() && SlateTexture->IsValid())
+	{
+		TSharedPtr<FSlateTexture2DRHIRef, ESPMode::ThreadSafe> SlateTextureCleanUp = SlateTexture;
+		ENQUEUE_RENDER_COMMAND(FAVPlayerMovieStreamerCleanUp)
+		(
+			[SlateTextureCleanUp](FRHICommandListImmediate& RHICmdList)
+			{
+				SlateTextureCleanUp->ReleaseDynamicRHI();
+			}
+		);
+		
+		SlateTexture.Reset();
+	}
 }
 
 bool FAVPlayerMovieStreamer::StartNextMovie()
@@ -398,7 +437,12 @@ bool FAVPlayerMovieStreamer::FinishLoadingTracks()
 				
                 // Initialize our video output to match the format of the texture we'll create later.
                 NSMutableDictionary* OutputSettings = [NSMutableDictionary dictionary];
-                [OutputSettings setObject: [NSNumber numberWithInt:kCVPixelFormatType_32BGRA] forKey:(NSString*)kCVPixelBufferPixelFormatTypeKey];
+				
+                [OutputSettings setObject:[NSNumber numberWithInt:kCVPixelFormatType_32BGRA] forKey:(NSString*)kCVPixelBufferPixelFormatTypeKey];
+                [OutputSettings setObject:[NSDictionary dictionary] forKey:(id)kCVPixelBufferIOSurfacePropertiesKey];
+				[OutputSettings setObject:[NSNumber numberWithBool:YES] forKey:(NSString*)kCVPixelBufferMetalCompatibilityKey];
+				[OutputSettings setObject:[NSNumber numberWithInteger:1] forKey:(NSString*)kCVPixelBufferBytesPerRowAlignmentKey];
+				
                 AVVideoOutput = [[AVAssetReaderTrackOutput alloc] initWithTrack:AVVideoTrack outputSettings:OutputSettings];
                 AVVideoOutput.alwaysCopiesSampleData = NO;
 
@@ -519,51 +563,38 @@ bool FAVPlayerMovieStreamer::CheckForNextFrameAndCopy()
     if( SyncStatus == Ready )
     {
         check(LatestSamples != NULL);
+		check(MetalTextureCache != NULL);
 
-        // Grab the pixel buffer and lock it for reading.
         CVImageBufferRef pPixelBuffer = CMSampleBufferGetImageBuffer(LatestSamples);
-        CGSize Size                   = CVImageBufferGetEncodedSize( pPixelBuffer );
-        CVPixelBufferLockBaseAddress( pPixelBuffer, kCVPixelBufferLock_ReadOnly );
-
-        uint8* pVideoData = (uint8*)CVPixelBufferGetBaseAddress(pPixelBuffer);
-
-        uint32 SrcWidth  = (uint32)Size.width;
-        uint32 SrcHeight = (uint32)Size.height;
-
-        // Now that we have video information, check on texture allocation. If we don't have a texture yet, create one.
-        if(!Texture.IsValid() || (Texture->GetWidth() != SrcWidth || Texture->GetHeight() != SrcHeight))
-        {
-            MovieViewport->SetTexture(NULL);
-
-            // Release any resources associated with the previous texture
-            if (Texture.IsValid())
-            {
-                // *** Aren't we already in the rendering thread? Why can't we just release the resource right here, right now?
-                TexturesPendingDeletion.Add(Texture);
-                BeginReleaseResource(Texture.Get());
-                Texture.Reset();
-            }
-
-            // Create and initialize a new texture
-            Texture = MakeShareable(new FSlateTexture2DRHIRef(SrcWidth, SrcHeight, PF_B8G8R8A8, nullptr, TexCreate_Dynamic | TexCreate_NoTiling, true));
-            Texture->InitResource();
-
-            // Make sure the texture is at least updated once.
-            Texture->UpdateRHI();
-            MovieViewport->SetTexture(Texture);
-        }
 		
-		uint32 DataLen = SrcWidth * 4 * SrcHeight;
-        uint32 Stride;
+		CGSize Size = CVImageBufferGetEncodedSize(pPixelBuffer);
+		
+		uint32 Width  = (uint32)Size.width;
+		uint32 Height = (uint32)Size.height;
+	
+		CVMetalTextureRef TextureCacheRef = NULL;
+		CVReturn Result = CVMetalTextureCacheCreateTextureFromImage(kCFAllocatorDefault, MetalTextureCache, pPixelBuffer, nullptr, MTLPixelFormatBGRA8Unorm, Width, Height, 0, &TextureCacheRef);
+		check(Result == kCVReturnSuccess);
+		check(TextureCacheRef);
+	
+		FRHIResourceCreateInfo CreateInfo(new FAppleMovieStreamerTexture2DResourceWrapper(TextureCacheRef));
+		FTexture2DRHIRef MediaWrappedTexture = RHICreateTexture2D(Width, Height, PF_B8G8R8A8, 1, 1, TexCreate_Dynamic | TexCreate_NoTiling, CreateInfo);
 		
-		uint8* DestTextureData = (uint8*)RHILockTexture2D( Texture->GetTypedResource(), 0, RLM_WriteOnly, Stride, false );
-		FMemory::Memcpy( DestTextureData, pVideoData, DataLen );
-		RHIUnlockTexture2D( Texture->GetTypedResource(), 0, false );
+		if(SlateTexture.IsValid())
+		{
+			if(SlateTexture->IsValid())
+			{
+				SlateTexture->ReleaseDynamicRHI();
+			}
+			SlateTexture->SetRHIRef(MediaWrappedTexture, Width, Height);
+		}
+		else
+		{
+			SlateTexture = MakeShareable(new FSlateTexture2DRHIRef(MediaWrappedTexture, Width, Height));
+			MovieViewport->SetTexture(SlateTexture);
+		}
 
-        // Re-lock and release the video data.
-        CVPixelBufferUnlockBaseAddress( pPixelBuffer, kCVPixelBufferLock_ReadOnly );
-        
-        // Processed this frame, so dump the samples.
+		CFRelease(TextureCacheRef);
         CFRelease(LatestSamples);
         LatestSamples = NULL;
 
@@ -587,12 +618,11 @@ void FAVPlayerMovieStreamer::TeardownPlayback()
     }
     
 	ReleaseMovie();
+	
     if ( AudioPlayer != nil )
     {
         AudioPlayer = nil;
     }
-
-    // NOTE: The any textures allocated are still allocated at this point. They will get released in Cleanup()
 }
 
 void FAVPlayerMovieStreamer::ReleaseMovie()
diff --git a/Engine/Plugins/Runtime/AppleMoviePlayer/Source/AppleMoviePlayer/Private/AppleMovieStreamer.h b/Engine/Plugins/Runtime/AppleMoviePlayer/Source/AppleMoviePlayer/Private/AppleMovieStreamer.h
--- a/Engine/Plugins/Runtime/AppleMoviePlayer/Source/AppleMoviePlayer/Private/AppleMovieStreamer.h
+++ b/Engine/Plugins/Runtime/AppleMoviePlayer/Source/AppleMoviePlayer/Private/AppleMovieStreamer.h
@@ -43,13 +43,10 @@ private:
         Ready,      // Frame is within tolerance of playback cursor.
     };
 
-    // Holds references to textures until their RHI resources are freed
-    TArray<TSharedPtr<FSlateTexture2DRHIRef, ESPMode::ThreadSafe>> TexturesPendingDeletion;
-
 	/** Texture and viewport data for displaying to Slate */
 	TSharedPtr<FMovieViewport> MovieViewport;
 
-    TSharedPtr<FSlateTexture2DRHIRef, ESPMode::ThreadSafe> Texture;
+    TSharedPtr<FSlateTexture2DRHIRef, ESPMode::ThreadSafe> SlateTexture;
 
     // The list of pending movies
     TArray<FString>		        MovieQueue;
@@ -77,6 +74,8 @@ private:
 	CMTime						ResumeTime;
 
 	FCriticalSection			VideoTracksLoadingLock;
+	
+	CVMetalTextureCacheRef 		MetalTextureCache;
 
 private:
 
