DR - Compile fix
#rnx

diff --git a/Engine/Plugins/Runtime/AppleMoviePlayer/Source/AppleMoviePlayer/Private/AppleMovieStreamer.cpp b/Engine/Plugins/Runtime/AppleMoviePlayer/Source/AppleMoviePlayer/Private/AppleMovieStreamer.cpp
--- a/Engine/Plugins/Runtime/AppleMoviePlayer/Source/AppleMoviePlayer/Private/AppleMovieStreamer.cpp
+++ b/Engine/Plugins/Runtime/AppleMoviePlayer/Source/AppleMoviePlayer/Private/AppleMovieStreamer.cpp
@@ -9,6 +9,7 @@
 #include "Misc/ScopeLock.h"
 #include "Misc/Paths.h"
 #include "Misc/CoreDelegates.h"
+#include "Containers/ResourceArray.h"
 
 DEFINE_LOG_CATEGORY_STATIC(LogMoviePlayer, Log, All);
 
