Fixing Mac nonunity errors.

diff --git a/Engine/Plugins/Media/ElectraUtil/Source/ElectraSamples/Private/ElectraTextureSampleAVF.cpp b/Engine/Plugins/Media/ElectraUtil/Source/ElectraSamples/Private/ElectraTextureSampleAVF.cpp
--- a/Engine/Plugins/Media/ElectraUtil/Source/ElectraSamples/Private/ElectraTextureSampleAVF.cpp
+++ b/Engine/Plugins/Media/ElectraUtil/Source/ElectraSamples/Private/ElectraTextureSampleAVF.cpp
@@ -13,6 +13,7 @@
 #include "RHIStaticStates.h"
 #include "MediaShaders.h"
 #include "Containers/ResourceArray.h"
+#include "PipelineStateCache.h"
 #else
 #include "Containers/Array.h"
 #endif
diff --git a/Engine/Plugins/Runtime/AppleMoviePlayer/Source/AppleMoviePlayer/Private/AppleMovieStreamer.cpp b/Engine/Plugins/Runtime/AppleMoviePlayer/Source/AppleMoviePlayer/Private/AppleMovieStreamer.cpp
--- a/Engine/Plugins/Runtime/AppleMoviePlayer/Source/AppleMoviePlayer/Private/AppleMovieStreamer.cpp
+++ b/Engine/Plugins/Runtime/AppleMoviePlayer/Source/AppleMoviePlayer/Private/AppleMovieStreamer.cpp
@@ -3,6 +3,7 @@
 #include "AppleMovieStreamer.h"
 
 #include "Rendering/RenderingCommon.h"
+#include "RenderingThread.h"
 #include "Slate/SlateTextures.h"
 #include "MoviePlayer.h"
 #include "Misc/CommandLine.h"
